Change minimum iOS version to 13 on remaining template projects #rb jack.porter #jira UE-122042 #okforgithub public

diff --git a/Templates/TP_AEC_HandheldARBP/Config/DefaultEngine.ini b/Templates/TP_AEC_HandheldARBP/Config/DefaultEngine.ini
--- a/Templates/TP_AEC_HandheldARBP/Config/DefaultEngine.ini
+++ b/Templates/TP_AEC_HandheldARBP/Config/DefaultEngine.ini
@@ -109,7 +109,7 @@ MobileProvision=iOS11.mobileprovision
 bSupportsPortraitOrientation=True
 bSupportsUpsideDownOrientation=False
 FrameRateLock=PUFRL_60
-MinimumiOSVersion=IOS_12
+MinimumiOSVersion=IOS_13
 
 [/Script/Engine.Engine]
 bUseFixedFrameRate=True
diff --git a/Templates/TP_FirstPerson/Config/DefaultEngine.ini b/Templates/TP_FirstPerson/Config/DefaultEngine.ini
--- a/Templates/TP_FirstPerson/Config/DefaultEngine.ini
+++ b/Templates/TP_FirstPerson/Config/DefaultEngine.ini
@@ -17,5 +17,5 @@ GlobalDefaultGameMode=/Script/TP_FirstPerson.TP_FirstPersonGameMode
 GlobalDefaultServerGameMode=None
 
 [/Script/IOSRuntimeSettings.IOSRuntimeSettings]
-MinimumiOSVersion=IOS_12
+MinimumiOSVersion=IOS_13
 
diff --git a/Templates/TP_FirstPersonBP/Config/DefaultEngine.ini b/Templates/TP_FirstPersonBP/Config/DefaultEngine.ini
--- a/Templates/TP_FirstPersonBP/Config/DefaultEngine.ini
+++ b/Templates/TP_FirstPersonBP/Config/DefaultEngine.ini
@@ -17,5 +17,5 @@ GlobalDefaultGameMode=/Game/FirstPersonBP/Blueprints/FirstPersonGameMode.FirstPe
 GlobalDefaultServerGameMode=None
 
 [/Script/IOSRuntimeSettings.IOSRuntimeSettings]
-MinimumiOSVersion=IOS_12
+MinimumiOSVersion=IOS_13
 
diff --git a/Templates/TP_MFG_HandheldARBP/Config/DefaultEngine.ini b/Templates/TP_MFG_HandheldARBP/Config/DefaultEngine.ini
--- a/Templates/TP_MFG_HandheldARBP/Config/DefaultEngine.ini
+++ b/Templates/TP_MFG_HandheldARBP/Config/DefaultEngine.ini
@@ -109,7 +109,7 @@ MobileProvision=iOS11.mobileprovision
 bSupportsPortraitOrientation=True
 bSupportsUpsideDownOrientation=False
 FrameRateLock=PUFRL_60
-MinimumiOSVersion=IOS_12
+MinimumiOSVersion=IOS_13
 
 [/Script/Engine.Engine]
 bUseFixedFrameRate=True
