Updated TemplateDefs for HandheldARBP, MFG_HandheldARBP and AEC_HandheldARBP to update their Chinese template descriptions

diff --git a/Templates/TP_AEC_HandheldARBP/Config/TemplateDefs.ini b/Templates/TP_AEC_HandheldARBP/Config/TemplateDefs.ini
--- a/Templates/TP_AEC_HandheldARBP/Config/TemplateDefs.ini
+++ b/Templates/TP_AEC_HandheldARBP/Config/TemplateDefs.ini
@@ -12,7 +12,7 @@ LocalizedDisplayNames=(Language="ja",Text="携帯 AR")
 LocalizedDescriptions=(Language="ja",Text="Android と iOS 向けの拡張現実アプリケーション作成の土台です。AR モードのオンオフを行う実行時のロジック、平面検出に関するデバッグ情報、ヒット検出のサンプルコード、ライト推定などを含んでいます。")
 
 LocalizedDisplayNames=(Language="zh-Hans",Text="手持式AR应用")
-LocalizedDescriptions=(Language="zh-Hans",Text="为安卓及iOS设备创建增强现实应用的良好起点，包含开启和关闭AR模式的运行逻辑，关于平面检测的调试信息，命中检测及处理预计光照的示例代码。")
+LocalizedDescriptions=(Language="zh-Hans",Text="该模板是安卓及iOS设备创建增强现实应用的良好起点，包含开启和关闭AR模式的运行逻辑、关于平面检测的调试信息、命中检测及处理预计光照的示例代码。")
 
 ClassTypes=Pawn, GameMode, Widget, Actor
 AssetTypes=Static Mesh
diff --git a/Templates/TP_HandheldARBP/Config/TemplateDefs.ini b/Templates/TP_HandheldARBP/Config/TemplateDefs.ini
--- a/Templates/TP_HandheldARBP/Config/TemplateDefs.ini
+++ b/Templates/TP_HandheldARBP/Config/TemplateDefs.ini
@@ -12,7 +12,7 @@ LocalizedDisplayNames=(Language="ja",Text="携帯 AR")
 LocalizedDescriptions=(Language="ja",Text="Android と iOS 向けの拡張現実アプリケーション作成の土台です。AR モードのオンオフを行う実行時のロジック、平面検出に関するデバッグ情報、ヒット検出のサンプルコード、ライト推定などを含んでいます。")
 
 LocalizedDisplayNames=(Language="zh-Hans",Text="手持式AR应用")
-LocalizedDescriptions=(Language="zh-Hans",Text="为安卓及iOS设备创建增强现实应用的良好起点，包含开启和关闭AR模式的运行逻辑，关于平面检测的调试信息，命中检测及处理预计光照的示例代码。")
+LocalizedDescriptions=(Language="zh-Hans",Text="该模板是安卓及iOS设备创建增强现实应用的良好起点，包含开启和关闭AR模式的运行逻辑、关于平面检测的调试信息、命中检测及处理预计光照的示例代码。")
 
 ClassTypes=Pawn, GameMode, Widget, Actor
 AssetTypes=Static Mesh
diff --git a/Templates/TP_MFG_HandheldARBP/Config/TemplateDefs.ini b/Templates/TP_MFG_HandheldARBP/Config/TemplateDefs.ini
--- a/Templates/TP_MFG_HandheldARBP/Config/TemplateDefs.ini
+++ b/Templates/TP_MFG_HandheldARBP/Config/TemplateDefs.ini
@@ -12,7 +12,7 @@ LocalizedDisplayNames=(Language="ja",Text="携帯 AR")
 LocalizedDescriptions=(Language="ja",Text="Android と iOS 向けの拡張現実アプリケーション作成の土台です。AR モードのオンオフを行う実行時のロジック、平面検出に関するデバッグ情報、ヒット検出のサンプルコード、ライト推定などを含んでいます。")
 
 LocalizedDisplayNames=(Language="zh-Hans",Text="手持式AR应用")
-LocalizedDescriptions=(Language="zh-Hans",Text="为安卓及iOS设备创建增强现实应用的良好起点，包含开启和关闭AR模式的运行逻辑，关于平面检测的调试信息，命中检测及处理预计光照的示例代码。")
+LocalizedDescriptions=(Language="zh-Hans",Text="该模板是安卓及iOS设备创建增强现实应用的良好起点，包含开启和关闭AR模式的运行逻辑、关于平面检测的调试信息、命中检测及处理预计光照的示例代码。")
 
 ClassTypes=Pawn, GameMode, Widget, Actor
 AssetTypes=Static Mesh
