Removed some of the obsolete mobile specific cvars from templates config files
r.Mobile.AllowMovableDirectionalLights
r.Mobile.EnableNoPrecomputedLightingCSMShader
r.MobileNumDynamicPointLights

diff --git a/Templates/TP_AEC_HandheldARBP/Config/DefaultEngine.ini b/Templates/TP_AEC_HandheldARBP/Config/DefaultEngine.ini
--- a/Templates/TP_AEC_HandheldARBP/Config/DefaultEngine.ini
+++ b/Templates/TP_AEC_HandheldARBP/Config/DefaultEngine.ini
@@ -74,9 +74,6 @@ r.SkinCache.CompileShaders=False
 r.Mobile.EnableStaticAndCSMShadowReceivers=True
 r.Mobile.EnableMovableLightCSMShaderCulling=True
 r.Mobile.AllowDistanceFieldShadows=True
-r.Mobile.AllowMovableDirectionalLights=True
-r.MobileNumDynamicPointLights=4
-r.MobileDynamicPointLightsUseStaticBranch=True
 r.SkinCache.SceneMemoryLimitInMB=128.000000
 r.GPUSkin.Limit2BoneInfluences=False
 r.SupportDepthOnlyIndexBuffers=True
diff --git a/Templates/TP_Blank/Config/DefaultEngine.ini b/Templates/TP_Blank/Config/DefaultEngine.ini
--- a/Templates/TP_Blank/Config/DefaultEngine.ini
+++ b/Templates/TP_Blank/Config/DefaultEngine.ini
@@ -32,9 +32,6 @@ CompressionQualityModifier=1.000000
 AutoStreamingThreshold=0.000000
 SoundCueCookQualityIndex=-1
 
-[/Script/Engine.RendererSettings]
-r.Mobile.EnableNoPrecomputedLightingCSMShader=True
-
 [/Script/LinuxTargetPlatform.LinuxTargetSettings]
 -TargetedRHIs=SF_VULKAN_SM5
 +TargetedRHIs=SF_VULKAN_SM6
\ No newline at end of file
diff --git a/Templates/TP_BlankBP/Config/DefaultEngine.ini b/Templates/TP_BlankBP/Config/DefaultEngine.ini
--- a/Templates/TP_BlankBP/Config/DefaultEngine.ini
+++ b/Templates/TP_BlankBP/Config/DefaultEngine.ini
@@ -32,9 +32,6 @@ CompressionQualityModifier=1.000000
 AutoStreamingThreshold=0.000000
 SoundCueCookQualityIndex=-1
 
-[/Script/Engine.RendererSettings]
-r.Mobile.EnableNoPrecomputedLightingCSMShader=True
-
 [/Script/LinuxTargetPlatform.LinuxTargetSettings]
 -TargetedRHIs=SF_VULKAN_SM5
 +TargetedRHIs=SF_VULKAN_SM6
\ No newline at end of file
diff --git a/Templates/TP_FirstPerson/Config/DefaultEngine.ini b/Templates/TP_FirstPerson/Config/DefaultEngine.ini
--- a/Templates/TP_FirstPerson/Config/DefaultEngine.ini
+++ b/Templates/TP_FirstPerson/Config/DefaultEngine.ini
@@ -20,7 +20,6 @@ GlobalDefaultServerGameMode=None
 r.ReflectionMethod=1
 r.GenerateMeshDistanceFields=True
 r.DynamicGlobalIlluminationMethod=1
-r.Mobile.EnableNoPrecomputedLightingCSMShader=1
 r.DefaultFeature.AutoExposure.ExtendDefaultLuminanceRange=true
 
 
diff --git a/Templates/TP_FirstPersonBP/Config/DefaultEngine.ini b/Templates/TP_FirstPersonBP/Config/DefaultEngine.ini
--- a/Templates/TP_FirstPersonBP/Config/DefaultEngine.ini
+++ b/Templates/TP_FirstPersonBP/Config/DefaultEngine.ini
@@ -20,7 +20,6 @@ GlobalDefaultServerGameMode=None
 r.ReflectionMethod=1
 r.GenerateMeshDistanceFields=True
 r.DynamicGlobalIlluminationMethod=1
-r.Mobile.EnableNoPrecomputedLightingCSMShader=1
 r.DefaultFeature.AutoExposure.ExtendDefaultLuminanceRange=true
 
 [/Script/WindowsTargetPlatform.WindowsTargetSettings]
diff --git a/Templates/TP_HandheldARBP/Config/DefaultEngine.ini b/Templates/TP_HandheldARBP/Config/DefaultEngine.ini
--- a/Templates/TP_HandheldARBP/Config/DefaultEngine.ini
+++ b/Templates/TP_HandheldARBP/Config/DefaultEngine.ini
@@ -73,9 +73,6 @@ r.RayTracing=False
 r.Mobile.EnableStaticAndCSMShadowReceivers=True
 r.Mobile.EnableMovableLightCSMShaderCulling=True
 r.Mobile.AllowDistanceFieldShadows=True
-r.Mobile.AllowMovableDirectionalLights=True
-r.MobileNumDynamicPointLights=4
-r.MobileDynamicPointLightsUseStaticBranch=True
 r.SkinCache.SceneMemoryLimitInMB=128.000000
 r.GPUSkin.Limit2BoneInfluences=False
 r.SupportDepthOnlyIndexBuffers=True
diff --git a/Templates/TP_MFG_HandheldARBP/Config/DefaultEngine.ini b/Templates/TP_MFG_HandheldARBP/Config/DefaultEngine.ini
--- a/Templates/TP_MFG_HandheldARBP/Config/DefaultEngine.ini
+++ b/Templates/TP_MFG_HandheldARBP/Config/DefaultEngine.ini
@@ -73,9 +73,6 @@ r.SkinCache.CompileShaders=False
 r.Mobile.EnableStaticAndCSMShadowReceivers=True
 r.Mobile.EnableMovableLightCSMShaderCulling=True
 r.Mobile.AllowDistanceFieldShadows=True
-r.Mobile.AllowMovableDirectionalLights=True
-r.MobileNumDynamicPointLights=4
-r.MobileDynamicPointLightsUseStaticBranch=True
 r.SkinCache.SceneMemoryLimitInMB=128.000000
 r.GPUSkin.Limit2BoneInfluences=False
 r.SupportDepthOnlyIndexBuffers=True
diff --git a/Templates/TP_ThirdPerson/Config/DefaultEngine.ini b/Templates/TP_ThirdPerson/Config/DefaultEngine.ini
--- a/Templates/TP_ThirdPerson/Config/DefaultEngine.ini
+++ b/Templates/TP_ThirdPerson/Config/DefaultEngine.ini
@@ -9,7 +9,6 @@ r.GenerateMeshDistanceFields=True
 r.DynamicGlobalIlluminationMethod=1
 r.Lumen.TraceMeshSDFs=0
 r.Shadow.Virtual.Enable=1
-r.Mobile.EnableNoPrecomputedLightingCSMShader=1
 r.DefaultFeature.AutoExposure.ExtendDefaultLuminanceRange=true
 
 [/Script/WindowsTargetPlatform.WindowsTargetSettings]
diff --git a/Templates/TP_ThirdPersonBP/Config/DefaultEngine.ini b/Templates/TP_ThirdPersonBP/Config/DefaultEngine.ini
--- a/Templates/TP_ThirdPersonBP/Config/DefaultEngine.ini
+++ b/Templates/TP_ThirdPersonBP/Config/DefaultEngine.ini
@@ -17,7 +17,6 @@ r.GenerateMeshDistanceFields=True
 r.DynamicGlobalIlluminationMethod=1
 r.Lumen.TraceMeshSDFs=0
 r.Shadow.Virtual.Enable=1
-r.Mobile.EnableNoPrecomputedLightingCSMShader=1
 r.DefaultFeature.AutoExposure.ExtendDefaultLuminanceRange=true
 
 [/Script/WindowsTargetPlatform.WindowsTargetSettings]
diff --git a/Templates/TP_TopDown/Config/DefaultEngine.ini b/Templates/TP_TopDown/Config/DefaultEngine.ini
--- a/Templates/TP_TopDown/Config/DefaultEngine.ini
+++ b/Templates/TP_TopDown/Config/DefaultEngine.ini
@@ -4,8 +4,7 @@ EditorStartupMap=/Game/TopDown/Maps/TopDownMap.TopDownMap
 GlobalDefaultGameMode="/Script/TP_TopDown.TP_TopDownGameMode"
 
 [/Script/Engine.RendererSettings]
-r.Mobile.EnableNoPrecomputedLightingCSMShader=1
- r.DefaultFeature.AutoExposure.ExtendDefaultLuminanceRange=true
+r.DefaultFeature.AutoExposure.ExtendDefaultLuminanceRange=true
 
 [/Script/WindowsTargetPlatform.WindowsTargetSettings]
 DefaultGraphicsRHI=DefaultGraphicsRHI_DX12
diff --git a/Templates/TP_TopDownBP/Config/DefaultEngine.ini b/Templates/TP_TopDownBP/Config/DefaultEngine.ini
--- a/Templates/TP_TopDownBP/Config/DefaultEngine.ini
+++ b/Templates/TP_TopDownBP/Config/DefaultEngine.ini
@@ -54,7 +54,6 @@ r.ReflectionMethod=1
 r.GenerateMeshDistanceFields=True
 r.DynamicGlobalIlluminationMethod=1
 r.Shadow.Virtual.Enable=1
-r.Mobile.EnableNoPrecomputedLightingCSMShader=1
 r.DefaultFeature.AutoExposure.ExtendDefaultLuminanceRange=true
 
 [/Script/WindowsTargetPlatform.WindowsTargetSettings]
diff --git a/Templates/TP_VehicleAdv/Config/DefaultEngine.ini b/Templates/TP_VehicleAdv/Config/DefaultEngine.ini
--- a/Templates/TP_VehicleAdv/Config/DefaultEngine.ini
+++ b/Templates/TP_VehicleAdv/Config/DefaultEngine.ini
@@ -19,7 +19,6 @@ GlobalDefaultServerGameMode=None
 r.ReflectionMethod=1
 r.GenerateMeshDistanceFields=True
 r.DynamicGlobalIlluminationMethod=1
-r.Mobile.EnableNoPrecomputedLightingCSMShader=1
 
 [/Script/WindowsTargetPlatform.WindowsTargetSettings]
 DefaultGraphicsRHI=DefaultGraphicsRHI_DX12
diff --git a/Templates/TP_VehicleAdvBP/Config/DefaultEngine.ini b/Templates/TP_VehicleAdvBP/Config/DefaultEngine.ini
--- a/Templates/TP_VehicleAdvBP/Config/DefaultEngine.ini
+++ b/Templates/TP_VehicleAdvBP/Config/DefaultEngine.ini
@@ -22,7 +22,6 @@ GlobalDefaultServerGameMode=None
 r.ReflectionMethod=1
 r.GenerateMeshDistanceFields=True
 r.DynamicGlobalIlluminationMethod=1
-r.Mobile.EnableNoPrecomputedLightingCSMShader=1
 r.ClearCoatNormal=True
 r.DefaultFeature.AutoExposure.ExtendDefaultLuminanceRange=true
 
