HandHeldAR Template | Disabled OpenGL support and pdated .ini settings to match engine defaults

diff --git a/Templates/TP_AEC_HandheldARBP/Config/DefaultEngine.ini b/Templates/TP_AEC_HandheldARBP/Config/DefaultEngine.ini
--- a/Templates/TP_AEC_HandheldARBP/Config/DefaultEngine.ini
+++ b/Templates/TP_AEC_HandheldARBP/Config/DefaultEngine.ini
@@ -1,84 +1,3 @@
-[URL]
-[/Script/Engine.RendererSettings]
-r.MobileHDR=True
-r.Mobile.DisableVertexFog=True
-
-r.MobileMSAA=1
-r.DiscardUnusedQuality=False
-r.AllowOcclusionQueries=True
-r.MinScreenRadiusForLights=0.030000
-r.MinScreenRadiusForDepthPrepass=0.030000
-r.MinScreenRadiusForCSMDepth=0.010000
-r.PrecomputedVisibilityWarning=False
-r.TextureStreaming=True
-Compat.UseDXT5NormalMaps=False
-r.ClearCoatNormal=False
-r.ReflectionCaptureResolution=128
-r.ReflectionEnvironmentLightmapMixBasedOnRoughness=True
-r.ForwardShading=False
-r.VertexFoggingForOpaque=True
-r.AllowStaticLighting=True
-r.NormalMapsForStaticLighting=False
-r.GenerateMeshDistanceFields=False
-r.DistanceFieldBuild.EightBit=False
-r.GenerateLandscapeGIData=False
-r.DistanceFieldBuild.Compress=False
-r.TessellationAdaptivePixelsPerTriangle=48.000000
-r.SeparateTranslucency=False
-r.TranslucentSortPolicy=0
-TranslucentSortAxis=(X=0.000000,Y=-1.000000,Z=0.000000)
-r.CustomDepth=1
-r.CustomDepthTemporalAAJitter=True
-r.PostProcessing.PropagateAlpha=0
-r.UsePreExposure=False
-r.DefaultFeature.Bloom=True
-r.DefaultFeature.AmbientOcclusion=False
-r.DefaultFeature.AmbientOcclusionStaticFraction=True
-r.DefaultFeature.AutoExposure=False
-r.DefaultFeature.AutoExposure.Method=0
-r.DefaultFeature.LocalExposure.HighlightContrastScale=1.0
-r.DefaultFeature.LocalExposure.ShadowContrastScale=1.0
-r.DefaultFeature.MotionBlur=False
-r.DefaultFeature.LensFlare=False
-r.TemporalAA.Upsampling=False
-r.DefaultFeature.AntiAliasing=0
-r.DefaultFeature.PointLightUnits=1
-r.DefaultFeature.SpotLightUnits=1
-r.Shadow.UnbuiltPreviewInGame=True
-r.StencilForLODDither=False
-r.EarlyZPass=3
-r.EarlyZPassMovable=True
-r.EarlyZPassOnlyMaterialMasking=False
-r.DBuffer=True
-r.ClearSceneMethod=1
-r.BasePassOutputsVelocity=False
-r.SelectiveBasePassOutputs=False
-bDefaultParticleCutouts=False
-r.AllowGlobalClipPlane=False
-r.GBufferFormat=1
-r.MorphTarget.Mode=True
-bGPUParticlesLocalVFOnly=False
-r.GPUCrashDebugging=False
-vr.InstancedStereo=False
-vr.MultiView=False
-vr.MobileMultiView=False
-vr.MobileMultiView.Direct=False
-vr.MonoscopicFarField=False
-vr.DebugCanvasInLayer=False
-r.WireframeCullThreshold=5.000000
-r.SupportStationarySkylight=True
-r.SupportLowQualityLightmaps=True
-r.SupportPointLightWholeSceneShadows=True
-r.SupportAtmosphericFog=True
-r.SkinCache.CompileShaders=False
-r.Mobile.EnableStaticAndCSMShadowReceivers=True
-r.Mobile.EnableMovableLightCSMShaderCulling=True
-r.Mobile.AllowDistanceFieldShadows=True
-r.SkinCache.SceneMemoryLimitInMB=128.000000
-r.GPUSkin.Limit2BoneInfluences=False
-r.SupportDepthOnlyIndexBuffers=True
-r.SupportReversedIndexBuffers=True
-
 [/Script/EngineSettings.GameMapsSettings]
 EditorStartupMap=/Game/HandheldAR/Maps/HandheldARBlankMap.HandheldARBlankMap
 LocalMapOptions=
@@ -93,79 +12,7 @@ ServerDefaultMap=/Engine/Maps/Entry.Entry
 GlobalDefaultGameMode=/Game/HandheldAR/Blueprints/GameFramework/BP_ARGameMode.BP_ARGameMode_C
 GlobalDefaultServerGameMode=None
 
-[/Script/Slate.SlateSettings]
-bExplicitCanvasChildZOrder=True
-
-[/Script/HardwareTargeting.HardwareTargetingSettings]
-TargetedHardwareClass=Mobile
-AppliedTargetedHardwareClass=Mobile
-DefaultGraphicsPerformance=Maximum
-AppliedDefaultGraphicsPerformance=Maximum
-
-[/Script/IOSRuntimeSettings.IOSRuntimeSettings]
-SigningCertificate=iPhone Developer: Chance Ivey (AP6BVS76S9)
-MobileProvision=iOS11.mobileprovision
-bSupportsPortraitOrientation=True
-bSupportsUpsideDownOrientation=False
-FrameRateLock=PUFRL_60
-
-[/Script/Engine.Engine]
-bUseFixedFrameRate=True
-FixedFrameRate=60.000000
-GameScreenshotSaveDirectory=(Path="Screenshots/")
-
 [/Script/AndroidRuntimeSettings.AndroidRuntimeSettings]
 Orientation=Sensor
-MinSDKVersion=24
-TargetSDKVersion=24
-
-[/Script/OculusHMD.OculusHMDRuntimeSettings]
-bAutoEnabled=True
-
-[/Script/Engine.PhysicsSettings]
-DefaultGravityZ=-980.000000
-DefaultTerminalVelocity=4000.000000
-DefaultFluidFriction=0.300000
-SimulateScratchMemorySize=262144
-RagdollAggregateThreshold=4
-TriangleMeshTriangleMinAreaThreshold=5.000000
-bEnableAsyncScene=False
-bEnableShapeSharing=False
-bEnablePCM=True
-bEnableStabilization=False
-bWarnMissingLocks=True
-bEnable2DPhysics=False
-LockedAxis=Invalid
-DefaultDegreesOfFreedom=Full3D
-BounceThresholdVelocity=200.000000
-FrictionCombineMode=Average
-RestitutionCombineMode=Average
-MaxAngularVelocity=3600.000000
-MaxDepenetrationVelocity=0.000000
-ContactOffsetMultiplier=0.020000
-MinContactOffset=2.000000
-MaxContactOffset=8.000000
-bSimulateSkeletalMeshOnDedicatedServer=True
-DefaultShapeComplexity=CTF_UseSimpleAndComplex
-bDefaultHasComplexCollision=True
-bSuppressFaceRemapTable=False
-bSupportUVFromHitResults=False
-bDisableActiveActors=False
-bDisableCCD=False
-bEnableEnhancedDeterminism=False
-MaxPhysicsDeltaTime=0.033333
-bSubstepping=False
-bSubsteppingAsync=False
-MaxSubstepDeltaTime=0.016667
-MaxSubsteps=6
-SyncSceneSmoothingFactor=0.000000
-AsyncSceneSmoothingFactor=0.990000
-InitialAverageFrameRate=0.016667
-PhysXTreeRebuildRate=10
+bBuildForES31=False
 
-[/Script/WindowsTargetPlatform.WindowsTargetSettings]
-DefaultGraphicsRHI=DefaultGraphicsRHI_DX12
--D3D12TargetedShaderFormats=PCD3D_SM5
-+D3D12TargetedShaderFormats=PCD3D_SM6
--D3D11TargetedShaderFormats=PCD3D_SM5
-+D3D11TargetedShaderFormats=PCD3D_SM5
